OpenXREyeTracker plugin definition cleanup
-Removed plugin dependency platform specificity from OpenXREyeTracker.uplugin.  The plugin dependency is for all platforms, and over-specifying these just makes them more fragile.

diff --git a/Engine/Plugins/Runtime/OpenXREyeTracker/OpenXREyeTracker.uplugin b/Engine/Plugins/Runtime/OpenXREyeTracker/OpenXREyeTracker.uplugin
--- a/Engine/Plugins/Runtime/OpenXREyeTracker/OpenXREyeTracker.uplugin
+++ b/Engine/Plugins/Runtime/OpenXREyeTracker/OpenXREyeTracker.uplugin
@@ -36,13 +36,7 @@
 	"Plugins": [
 		{
 			"Name": "OpenXR",
-			"Enabled": true,
-			"PlatformAllowList": [
-				"Win64",
-				"Linux",
-				"HoloLens",
-				"Android"
-			]
+			"Enabled": true
 		}
 	]
 }
