OpenXREyeTracker updated deprecated functions
-It was using older, deprecated interfaces.  Now its using the latest stuff.

diff --git a/Engine/Plugins/Runtime/OpenXREyeTracker/Source/OpenXREyeTracker/Private/OpenXREyeTracker.cpp b/Engine/Plugins/Runtime/OpenXREyeTracker/Source/OpenXREyeTracker/Private/OpenXREyeTracker.cpp
--- a/Engine/Plugins/Runtime/OpenXREyeTracker/Source/OpenXREyeTracker/Private/OpenXREyeTracker.cpp
+++ b/Engine/Plugins/Runtime/OpenXREyeTracker/Source/OpenXREyeTracker/Private/OpenXREyeTracker.cpp
@@ -58,6 +58,11 @@ bool FOpenXREyeTracker::GetRequiredExtensions(TArray<const ANSICHAR*>& OutExtens
 	return true;
 }
 
+void FOpenXREyeTracker::PostCreateInstance(XrInstance InInstance)
+{
+	Instance = InInstance;
+}
+
 bool FOpenXREyeTracker::GetInteractionProfile(XrInstance InInstance, FString& OutKeyPrefix, XrPath& OutPath, bool& OutHasHaptics)
 {
 	OutKeyPrefix = "EyeTracker";
@@ -65,33 +70,74 @@ bool FOpenXREyeTracker::GetInteractionProfile(XrInstance InInstance, FString& Ou
 	return xrStringToPath(InInstance, "/interaction_profiles/ext/eye_gaze_interaction", &OutPath) == XR_SUCCESS;
 }
 
-void FOpenXREyeTracker::AddActions(XrInstance Instance, TFunction<XrAction(XrActionType InActionType, const FName& InName, const TArray<XrPath>& InSubactionPaths)> AddAction)
+void FOpenXREyeTracker::AttachActionSets(TSet<XrActionSet>& OutActionSets)
 {
-	TArray<XrPath> SubactionPaths;
+	check(Instance != XR_NULL_HANDLE);
+
+	// This is a bit of a pain right now.  Hopefully future refactors will make it better.
+	// We are creating an action set for our eye tracking pose.  An action set can have session create/destroy
+	// lifetime.  However currently the OpenXRInput module is loaded well after the OpenXRHMDModule creates 
+	// the session, so we can't just setup all this input system right then.  OpenXRInput instead checks if it is live and if 
+	// not destroys any existing sets and then creates new ones near xrBeginSession (where the session starts running).
+	// It marks them as dead near xrDestroySession so that they will be destroyed on the BeginSession of the next created session, 
+	// if any.
+	//
+	// To mirror that lifetime we are destroying and creating in AttachActionSets and marking as dead in OnDestroySession.
+	//
+	// Note the ActionSpace is easier because it is never used outside of this ExtensionPlugin.  We are creating it, if necessary, 
+	// in OnBeginSession and destroying it in OnDestroySession.  If we had a good CreateSession hook we could create it then, along with
+	// the ActionSet and Action.
+
+	// We could have an action set from a previous session.  If so it needs to go away.
+	if (EyeTrackerActionSet != XR_NULL_HANDLE)
+	{
+		xrDestroyActionSet(EyeTrackerActionSet);
+		EyeTrackerActionSet = XR_NULL_HANDLE;
+		EyeTrackerAction = XR_NULL_HANDLE;
+	}
 
-	EyeTrackerAction = AddAction(XR_ACTION_TYPE_POSE_INPUT, "eye_tracker", SubactionPaths);
-	if (EyeTrackerAction == XR_NULL_HANDLE)
 	{
-		return;
+		XrActionSetCreateInfo Info;
+		Info.type = XR_TYPE_ACTION_SET_CREATE_INFO;
+		Info.next = nullptr;
+		FCStringAnsi::Strcpy(Info.actionSetName, XR_MAX_ACTION_SET_NAME_SIZE, "openxreyetrackeractionset");
+		FCStringAnsi::Strcpy(Info.localizedActionSetName, XR_MAX_LOCALIZED_ACTION_SET_NAME_SIZE, "OpenXR Eye Tracker Action Set");
+		Info.priority = 0;
+		XR_ENSURE(xrCreateActionSet(Instance, &Info, &EyeTrackerActionSet));
 	}
 
-	// Create suggested bindings
-	XrPath EyeGazeInteractionProfilePath = XR_NULL_PATH;
-	XR_ENSURE(xrStringToPath(Instance, "/interaction_profiles/ext/eye_gaze_interaction", &EyeGazeInteractionProfilePath));
+	{
+		check(EyeTrackerAction == XR_NULL_HANDLE);
+		XrActionCreateInfo Info;
+		Info.type = XR_TYPE_ACTION_CREATE_INFO;
+		Info.next = nullptr;
+		Info.countSubactionPaths = 0;
+		Info.subactionPaths = nullptr;
+		FCStringAnsi::Strcpy(Info.actionName, XR_MAX_ACTION_NAME_SIZE, "openxreyetrackeraction");
+		FCStringAnsi::Strcpy(Info.localizedActionName, XR_MAX_LOCALIZED_ACTION_NAME_SIZE, "OpenXR Eye Tracker Action");
+		Info.actionType = XR_ACTION_TYPE_POSE_INPUT;
+		XR_ENSURE(xrCreateAction(EyeTrackerActionSet, &Info, &EyeTrackerAction));
+	}
+
+	{
+		XrPath EyeGazeInteractionProfilePath = XR_NULL_PATH;
+		XR_ENSURE(xrStringToPath(Instance, "/interaction_profiles/ext/eye_gaze_interaction", &EyeGazeInteractionProfilePath));
 
-	XrPath GazePosePath = XR_NULL_PATH;
-	XR_ENSURE(xrStringToPath(Instance, "/user/eyes_ext/input/gaze_ext/pose", &GazePosePath));
+		XrPath GazePosePath = XR_NULL_PATH;
+		XR_ENSURE(xrStringToPath(Instance, "/user/eyes_ext/input/gaze_ext/pose", &GazePosePath));
 
-	XrActionSuggestedBinding Bindings;
-	Bindings.action = EyeTrackerAction;
-	Bindings.binding = GazePosePath;
+		XrActionSuggestedBinding Bindings;
+		Bindings.action = EyeTrackerAction;
+		Bindings.binding = GazePosePath;
 
-	XrInteractionProfileSuggestedBinding SuggestedBindings{ XR_TYPE_INTERACTION_PROFILE_SUGGESTED_BINDING };
-	SuggestedBindings.interactionProfile = EyeGazeInteractionProfilePath;
-	SuggestedBindings.suggestedBindings = &Bindings;
-	SuggestedBindings.countSuggestedBindings = 1;
-	XR_ENSURE(xrSuggestInteractionProfileBindings(Instance, &SuggestedBindings));
+		XrInteractionProfileSuggestedBinding SuggestedBindings{ XR_TYPE_INTERACTION_PROFILE_SUGGESTED_BINDING };
+		SuggestedBindings.interactionProfile = EyeGazeInteractionProfilePath;
+		SuggestedBindings.suggestedBindings = &Bindings;
+		SuggestedBindings.countSuggestedBindings = 1;
+		XR_ENSURE(xrSuggestInteractionProfileBindings(Instance, &SuggestedBindings));
+	}
 
+	OutActionSets.Add(EyeTrackerActionSet);
 }
 
 const void* FOpenXREyeTracker::OnBeginSession(XrSession InSession, const void* InNext)
@@ -102,22 +148,44 @@ const void* FOpenXREyeTracker::OnBeginSession(XrSession InSession, const void* I
 		XRTrackingSystem = GEngine->XRSystem.Get();
 	}
 
-	XrActionSpaceCreateInfo CreateActionSpaceInfo{ XR_TYPE_ACTION_SPACE_CREATE_INFO };
-	CreateActionSpaceInfo.action = EyeTrackerAction;
-	CreateActionSpaceInfo.poseInActionSpace = ToXrPose(FTransform::Identity);
-	XR_ENSURE(xrCreateActionSpace(InSession, &CreateActionSpaceInfo, &GazeActionSpace));
-
-	SyncInfo.countActiveActionSets = 0;
-	SyncInfo.activeActionSets = XR_NULL_HANDLE;
+	if (GazeActionSpace == XR_NULL_HANDLE)
+	{
+		GazeActionSpace = XR_NULL_HANDLE;
+		XrActionSpaceCreateInfo CreateActionSpaceInfo{ XR_TYPE_ACTION_SPACE_CREATE_INFO };
+		check(EyeTrackerAction != XR_NULL_HANDLE);
+		CreateActionSpaceInfo.action = EyeTrackerAction;
+		CreateActionSpaceInfo.poseInActionSpace = ToXrPose(FTransform::Identity);
+		XR_ENSURE(xrCreateActionSpace(InSession, &CreateActionSpaceInfo, &GazeActionSpace));
+
+		SyncInfo.countActiveActionSets = 0;
+		SyncInfo.activeActionSets = XR_NULL_HANDLE;
+	}
 
 	bSessionStarted = true;
 
 	return InNext;
 }
 
+void FOpenXREyeTracker::OnDestroySession(XrSession InSession)
+{
+	if (GazeActionSpace)
+	{
+		XR_ENSURE(xrDestroySpace(GazeActionSpace));
+	}
+	GazeActionSpace = XR_NULL_HANDLE;
+}
+
+void FOpenXREyeTracker::GetActiveActionSetsForSync(TArray<XrActiveActionSet>& OutActiveSets)
+{
+	check(EyeTrackerActionSet != XR_NULL_HANDLE);
+
+	OutActiveSets.Add(XrActiveActionSet{ EyeTrackerActionSet, XR_NULL_PATH });
+}
 
 void FOpenXREyeTracker::PostSyncActions(XrSession InSession)
 {
+	check(EyeTrackerAction != XR_NULL_HANDLE);
+
 	XrActionStateGetInfo GetActionStateInfo{ XR_TYPE_ACTION_STATE_GET_INFO };
 	GetActionStateInfo.action = EyeTrackerAction;
 	XR_ENSURE(xrGetActionStatePose(InSession, &GetActionStateInfo, &ActionStatePose));
@@ -125,7 +193,9 @@ void FOpenXREyeTracker::PostSyncActions(XrSession InSession)
 
 void FOpenXREyeTracker::UpdateDeviceLocations(XrSession InSession, XrTime DisplayTime, XrSpace TrackingSpace)
 {
-	if (ActionStatePose.isActive) {
+	if (ActionStatePose.isActive) 
+	{
+		check(GazeActionSpace != XR_NULL_HANDLE);
 		XR_ENSURE(xrLocateSpace(GazeActionSpace, TrackingSpace, DisplayTime, &EyeTrackerSpaceLocation));
 	}
 }
diff --git a/Engine/Plugins/Runtime/OpenXREyeTracker/Source/OpenXREyeTracker/Private/OpenXREyeTracker.h b/Engine/Plugins/Runtime/OpenXREyeTracker/Source/OpenXREyeTracker/Private/OpenXREyeTracker.h
--- a/Engine/Plugins/Runtime/OpenXREyeTracker/Source/OpenXREyeTracker/Private/OpenXREyeTracker.h
+++ b/Engine/Plugins/Runtime/OpenXREyeTracker/Source/OpenXREyeTracker/Private/OpenXREyeTracker.h
@@ -35,10 +35,12 @@ public:
 		return FString(TEXT("OpenXREyeTracker"));
 	}
 	virtual bool GetRequiredExtensions(TArray<const ANSICHAR*>& OutExtensions) override;
+	virtual void PostCreateInstance(XrInstance InInstance) override;
 	virtual bool GetInteractionProfile(XrInstance InInstance, FString& OutKeyPrefix, XrPath& OutPath, bool& OutHasHaptics) override;
-	virtual void AddActions(XrInstance Instance, TFunction<XrAction(XrActionType InActionType, const FName& InName, const TArray<XrPath>& InSubactionPaths)> AddAction) override;
-
+	virtual void AttachActionSets(TSet<XrActionSet>& OutActionSets) override;
 	virtual const void* OnBeginSession(XrSession InSession, const void* InNext) override;
+	virtual void OnDestroySession(XrSession InSession) override;
+	virtual void GetActiveActionSetsForSync(TArray<XrActiveActionSet>& OutActiveSets) override;
 	virtual void PostSyncActions(XrSession InSession) override;
 	virtual void UpdateDeviceLocations(XrSession InSession, XrTime DisplayTime, XrSpace TrackingSpace) override;
 
@@ -46,11 +48,12 @@ public:
 
 private:
 	class IXRTrackingSystem* XRTrackingSystem = nullptr;
-
+	XrInstance Instance = XR_NULL_HANDLE;
 	bool bSessionStarted = false;
 	XrActionsSyncInfo SyncInfo{ XR_TYPE_ACTIONS_SYNC_INFO };
-	XrAction EyeTrackerAction;
-	XrSpace GazeActionSpace;
+	XrAction EyeTrackerAction = XR_NULL_HANDLE;
+	XrActionSet EyeTrackerActionSet = XR_NULL_HANDLE;
+	XrSpace GazeActionSpace = XR_NULL_HANDLE;
 	XrActionStatePose ActionStatePose{ XR_TYPE_ACTION_STATE_POSE };
 
 	// EyeTracker cached data
