UE-170928 OpenXREyeTracker tracked flag still true if tracking stops.
-Setting the tracked bits to false if the action is not active, which means the openxr space associated is not being tracked.

diff --git a/Engine/Plugins/Runtime/OpenXREyeTracker/Source/OpenXREyeTracker/Private/OpenXREyeTracker.cpp b/Engine/Plugins/Runtime/OpenXREyeTracker/Source/OpenXREyeTracker/Private/OpenXREyeTracker.cpp
--- a/Engine/Plugins/Runtime/OpenXREyeTracker/Source/OpenXREyeTracker/Private/OpenXREyeTracker.cpp
+++ b/Engine/Plugins/Runtime/OpenXREyeTracker/Source/OpenXREyeTracker/Private/OpenXREyeTracker.cpp
@@ -198,6 +198,12 @@ void FOpenXREyeTracker::UpdateDeviceLocations(XrSession InSession, XrTime Displa
 		check(GazeActionSpace != XR_NULL_HANDLE);
 		XR_ENSURE(xrLocateSpace(GazeActionSpace, TrackingSpace, DisplayTime, &EyeTrackerSpaceLocation));
 	}
+	else
+	{
+		// Clear the tracked bits if the action is not active
+		const XrSpaceLocationFlags TrackedFlags = XR_SPACE_LOCATION_ORIENTATION_TRACKED_BIT | XR_SPACE_LOCATION_POSITION_TRACKED_BIT;
+		EyeTrackerSpaceLocation.locationFlags &= ~TrackedFlags;
+	}
 }
 
 bool FOpenXREyeTracker::GetEyeTrackerGazeData(FEyeTrackerGazeData& OutGazeData) const
