Add Mac as a target/allow platform for the OpenXREyeTracker plugin

diff --git a/Engine/Plugins/Runtime/OpenXREyeTracker/OpenXREyeTracker.uplugin b/Engine/Plugins/Runtime/OpenXREyeTracker/OpenXREyeTracker.uplugin
--- a/Engine/Plugins/Runtime/OpenXREyeTracker/OpenXREyeTracker.uplugin
+++ b/Engine/Plugins/Runtime/OpenXREyeTracker/OpenXREyeTracker.uplugin
@@ -27,7 +27,8 @@
 			"PlatformAllowList": [
 				"Win64",
 				"Linux",
-				"Android"
+				"Android",
+				"Mac"
 			]
 		}
 	],
